Prevent fbx to be use by other platform then win64

diff --git a/Engine/Plugins/Runtime/Interchange/Interchange.uplugin b/Engine/Plugins/Runtime/Interchange/Interchange.uplugin
--- a/Engine/Plugins/Runtime/Interchange/Interchange.uplugin
+++ b/Engine/Plugins/Runtime/Interchange/Interchange.uplugin
@@ -38,7 +38,10 @@
 		{
 			"Name": "InterchangeFbxParser",
 			"Type": "Runtime",
-			"LoadingPhase": "Default"
+			"LoadingPhase": "Default",
+			"WhitelistPlatforms": [
+				"Win64"
+			]
 		}
 	]
 }
\ No newline at end of file
diff --git a/Engine/Source/Programs/InterchangeWorker/InterchangeWorker.Build.cs b/Engine/Source/Programs/InterchangeWorker/InterchangeWorker.Build.cs
--- a/Engine/Source/Programs/InterchangeWorker/InterchangeWorker.Build.cs
+++ b/Engine/Source/Programs/InterchangeWorker/InterchangeWorker.Build.cs
@@ -17,13 +17,17 @@ public class InterchangeWorker : ModuleRules
 				"ApplicationCore",
 				"Core",
 				"InterchangeDispatcher",
-				"InterchangeFbxParser",
 				"Json",
 				"Projects",
 				"Sockets",
 			}
 		);
 
+		if (Target.Platform == UnrealTargetPlatform.Win64)
+		{
+			PrivateDependencyModuleNames.Add("InterchangeFbxParser");
+		}
+
 		PublicDelayLoadDLLs.Add("kernel_io.dll");
 	}
 }
diff --git a/Engine/Source/Programs/InterchangeWorker/Private/InterchangeWorkerImpl.cpp b/Engine/Source/Programs/InterchangeWorker/Private/InterchangeWorkerImpl.cpp
--- a/Engine/Source/Programs/InterchangeWorker/Private/InterchangeWorkerImpl.cpp
+++ b/Engine/Source/Programs/InterchangeWorker/Private/InterchangeWorkerImpl.cpp
@@ -5,7 +5,10 @@
 #include "InterchangeCommands.h"
 #include "InterchangeDispatcherConfig.h"
 #include "InterchangeDispatcherTask.h"
+
+#if PLATFORM_WINDOWS && PLATFORM_64BITS
 #include "InterchangeFbxParser.h"
+#endif
 
 #include "HAL/FileManager.h"
 #include "Misc/FileHelper.h"
@@ -155,12 +158,21 @@ void FInterchangeWorkerImpl::ProcessCommand(const FRunTaskCommand& RunTaskComman
 
 ETaskState FInterchangeWorkerImpl::LoadFbxFile(const InterchangeDispatcher::FJsonLoadSourceCmd& LoadSourceCommand, FString& OutJSonResult, FString& OutJSonMessages) const
 {
+	ETaskState ResultState = ETaskState::Unknown;
 	FString SourceFilename = LoadSourceCommand.GetSourceFilename();
+#if PLATFORM_WINDOWS && PLATFORM_64BITS
 	InterchangeFbxParser::FbxParser FbxParser;
 	FbxParser.LoadFbxFile(SourceFilename, ResultFolder);
 	InterchangeDispatcher::FJsonLoadSourceCmd::JsonResultParser ResultParser;
 	ResultParser.SetResultFilename(FbxParser.GetResultFilepath());
 	OutJSonMessages = FbxParser.GetJsonLoadMessages();
 	OutJSonResult = ResultParser.ToJson();
-	return ETaskState::ProcessOk;
+	ResultState = ETaskState::ProcessOk;
+#else
+	InterchangeDispatcher::FJsonLoadSourceCmd::JsonResultParser ResultParser;
+	ResultParser.SetResultFilename(FString());
+	OutJSonMessages = TEXT("{\"Msg\" : {\"Type\" : \"Error\",\n\"Msg\" : \"Cannot Execute fbx command on other platform then window 64 bits!\"}}");
+	ResultState = ETaskState::ProcessFailed;
+#endif
+	return ResultState;
 }
\ No newline at end of file
