Use threadpool for the interchange worker network socket polling

diff --git a/Engine/Source/Programs/InterchangeWorker/Private/InterchangeWorkerImpl.cpp b/Engine/Source/Programs/InterchangeWorker/Private/InterchangeWorkerImpl.cpp
--- a/Engine/Source/Programs/InterchangeWorker/Private/InterchangeWorkerImpl.cpp
+++ b/Engine/Source/Programs/InterchangeWorker/Private/InterchangeWorkerImpl.cpp
@@ -2,12 +2,13 @@
 
 #include "InterchangeWorkerImpl.h"
 
+#include "Async/Async.h"
+#include "HAL/FileManager.h"
+#include "HAL/PlatformFileManager.h"
 #include "InterchangeCommands.h"
 #include "InterchangeDispatcherConfig.h"
 #include "InterchangeDispatcherTask.h"
 #include "InterchangeFbxParser.h"
-#include "HAL/FileManager.h"
-#include "HAL/PlatformFileManager.h"
 #include "Misc/FileHelper.h"
 #include "Misc/Paths.h"
 #include "SocketSubsystem.h"
@@ -67,7 +68,14 @@ bool FInterchangeWorkerImpl::Run()
 					//Use a randomGuid to generate thread unique name
 					int32 UniqueID = (FPlatformTime::Cycles64() & 0x00000000EFFFFFFF);
 					FString ThreadName = FString(TEXT("InterchangeWorkerCommand_")) + FString::FromInt(UniqueID);
-					IOThreads.Add(ThreadName, FThread(*ThreadName, [this, ThreadName, Command]() { ProcessCommand(Command, ThreadName); }));
+					ActiveThreads.Add(ThreadName, Async(
+						EAsyncExecution::ThreadPool,
+						[this, ThreadName, Command]()->bool
+						{
+							ProcessCommand(Command, ThreadName);
+							return true;
+						}
+					));
 				}
 					break;
 
@@ -95,11 +103,12 @@ bool FInterchangeWorkerImpl::Run()
 			FScopeLock Lock(&TFinishThreadCriticalSection);
 			for (const FString& ThreadName : CurrentFinishThreads)
 			{
-				if (FThread* IOThread = IOThreads.Find(ThreadName))
+				if (TFuture<bool>* IsThreadCompleted = ActiveThreads.Find(ThreadName))
 				{
-					IOThread->Join();
+					//Wait until process command is terminate
+					IsThreadCompleted->Get();
 				}
-				IOThreads.Remove(ThreadName);
+				ActiveThreads.Remove(ThreadName);
 			}
 			CurrentFinishThreads.Empty();
 		}
@@ -109,9 +118,9 @@ bool FInterchangeWorkerImpl::Run()
 	}
 
 	//Join all thread that is not terminate
-	for(TPair<FString, FThread>&IOThread : IOThreads)
+	for(TPair<FString, TFuture<bool>>& ActiveThreadPair : ActiveThreads)
 	{
-		IOThread.Value.Join();
+		ActiveThreadPair.Value.Get();
 	}
 
 	UE_CLOG(!bIsRunning, LogInterchangeWorker, Verbose, TEXT("Worker loop exit..."));
diff --git a/Engine/Source/Programs/InterchangeWorker/Private/InterchangeWorkerImpl.h b/Engine/Source/Programs/InterchangeWorker/Private/InterchangeWorkerImpl.h
--- a/Engine/Source/Programs/InterchangeWorker/Private/InterchangeWorkerImpl.h
+++ b/Engine/Source/Programs/InterchangeWorker/Private/InterchangeWorkerImpl.h
@@ -37,7 +37,7 @@ private:
 	FString ResultFolder;
 	FCriticalSection TFinishThreadCriticalSection;
 	TArray<FString> CurrentFinishThreads;
-	TMap<FString, FThread > IOThreads;
+	TMap<FString, TFuture<bool> > ActiveThreads;
 
 	UE::Interchange::FInterchangeFbxParser FbxParser;
 
