Make Interchange support LWC correctly
#rnx

diff --git a/Engine/Plugins/Runtime/Interchange/Source/Dispatcher/Private/InterchangeCommands.cpp b/Engine/Plugins/Runtime/Interchange/Source/Dispatcher/Private/InterchangeCommands.cpp
--- a/Engine/Plugins/Runtime/Interchange/Source/Dispatcher/Private/InterchangeCommands.cpp
+++ b/Engine/Plugins/Runtime/Interchange/Source/Dispatcher/Private/InterchangeCommands.cpp
@@ -15,6 +15,7 @@ namespace UE
 			switch (CommandType)
 			{
 				case ECommandId::Ping: return MakeShared<FPingCommand>();
+				case ECommandId::Error: return MakeShared<FErrorCommand>();
 				case ECommandId::BackPing: return MakeShared<FBackPingCommand>();
 				case ECommandId::RunTask: return MakeShared<FRunTaskCommand>();
 				case ECommandId::NotifyEndTask: return MakeShared<FCompletedTaskCommand>();
@@ -46,6 +47,11 @@ namespace UE
 			return nullptr;
 		}
 
+		void FErrorCommand::SerializeImpl(FArchive& Ar)
+		{
+			Ar << ErrorMessage;
+		}
+
 		void FRunTaskCommand::SerializeImpl(FArchive& Ar)
 		{
 			Ar << JsonDescription;
diff --git a/Engine/Plugins/Runtime/Interchange/Source/Dispatcher/Private/InterchangeWorkerHandler.cpp b/Engine/Plugins/Runtime/Interchange/Source/Dispatcher/Private/InterchangeWorkerHandler.cpp
--- a/Engine/Plugins/Runtime/Interchange/Source/Dispatcher/Private/InterchangeWorkerHandler.cpp
+++ b/Engine/Plugins/Runtime/Interchange/Source/Dispatcher/Private/InterchangeWorkerHandler.cpp
@@ -409,6 +409,10 @@ namespace UE
 				case ECommandId::Ping:
 					ProcessCommand(StaticCast<FPingCommand&>(Command));
 					break;
+				
+				case ECommandId::Error:
+					ProcessCommand(StaticCast<FErrorCommand&>(Command));
+					break;
 
 				case ECommandId::NotifyEndTask:
 					ProcessCommand(StaticCast<FCompletedTaskCommand&>(Command));
@@ -429,6 +433,14 @@ namespace UE
 			CommandIO.SendCommand(BackPing, 0);
 		}
 
+		void FInterchangeWorkerHandler::ProcessCommand(FErrorCommand& ErrorCommand)
+		{
+			UE_LOG(LogInterchangeDispatcher, Error, TEXT("%s"), *ErrorCommand.ErrorMessage);
+			this->Dispatcher.SetInterchangeWorkerFatalError(ErrorCommand.ErrorMessage);
+			FTerminateCommand Terminate;
+			CommandIO.SendCommand(Terminate, 0);
+		}
+
 		void FInterchangeWorkerHandler::ProcessCommand(FCompletedTaskCommand& CompletedTaskCommand)
 		{
 			{
diff --git a/Engine/Plugins/Runtime/Interchange/Source/Dispatcher/Public/InterchangeCommands.h b/Engine/Plugins/Runtime/Interchange/Source/Dispatcher/Public/InterchangeCommands.h
--- a/Engine/Plugins/Runtime/Interchange/Source/Dispatcher/Public/InterchangeCommands.h
+++ b/Engine/Plugins/Runtime/Interchange/Source/Dispatcher/Public/InterchangeCommands.h
@@ -23,24 +23,36 @@ namespace UE
 			//Patch version should be update if there is some bug fixes in the private code.
 			static int32 GetPatch() { return 0; }
 
+			//LWCDisabled version tell if the code was compile with serialization compatibility.
+			static bool GetLWCDisabled()
+			{
+#if UE_LARGE_WORLD_COORDINATES_DISABLED
+				return true;
+#else
+				return false;
+#endif
+			}
+
 			/** Return the version in string format "Major.Minor.Patch" */
 			static FString ToString()
 			{
-				static FString Version = FString::FromInt(GetMajor()) + TEXT(".") + FString::FromInt(GetMinor()) + TEXT(".") + FString::FromInt(GetPatch());
+				static FString Version = FString::FromInt(GetMajor()) + TEXT(".") + FString::FromInt(GetMinor()) + TEXT(".") + FString::FromInt(GetPatch()) + TEXT(".") + FString::FromInt(GetLWCDisabled() ? 1 : 0);
 				return Version;
 			}
 
-			static bool FromString(const FString& VersionStr, int32& OutMajor, int32& OutMinor, int32& OutPatch)
+			static bool FromString(const FString& VersionStr, int32& OutMajor, int32& OutMinor, int32& OutPatch, bool& bOutLWCDisabled)
 			{
 				TArray<FString> Tokens;
 				VersionStr.ParseIntoArray(Tokens, TEXT("."));
-				if (Tokens.Num() != 3)
+				if (Tokens.Num() != 4)
 				{
 					return false;
 				}
 				OutMajor = FCString::Atoi(*(Tokens[0]));
 				OutMinor = FCString::Atoi(*(Tokens[1]));
 				OutPatch = FCString::Atoi(*(Tokens[2]));
+				bOutLWCDisabled = FCString::Atoi(*(Tokens[3])) != 0;
+
 				return true;
 			}
 
@@ -53,15 +65,16 @@ namespace UE
 			 *
 			 * @return true if Major and Minor version equal the compile one, InPatch is only there in case we need it in a future release
 			 */
-			static bool IsAPICompatible(int32 Major, int32 Minor, int32 Patch)
+			static bool IsAPICompatible(int32 Major, int32 Minor, int32 Patch, bool bLWCDisabled)
 			{
-				return (GetMajor() == Major && GetMinor() == Minor);
+				return (GetMajor() == Major && GetMinor() == Minor && GetLWCDisabled() == bLWCDisabled);
 			}
 		};
 
 		enum class ECommandId : uint8
 		{
 			Invalid,
+			Error,
 			Ping,
 			BackPing,
 			RunTask,
@@ -105,6 +118,16 @@ namespace UE
 			virtual ECommandId GetType() const override { return ECommandId::Terminate; }
 		};
 
+		class INTERCHANGEDISPATCHER_API FErrorCommand : public ICommand
+		{
+		public:
+			virtual ECommandId GetType() const override { return ECommandId::Error; }
+			
+		protected:
+			virtual void SerializeImpl(FArchive&) override;
+		public:
+			FString ErrorMessage;
+		};
 
 		class INTERCHANGEDISPATCHER_API FPingCommand : public ICommand
 		{
diff --git a/Engine/Plugins/Runtime/Interchange/Source/Dispatcher/Public/InterchangeDispatcher.h b/Engine/Plugins/Runtime/Interchange/Source/Dispatcher/Public/InterchangeDispatcher.h
--- a/Engine/Plugins/Runtime/Interchange/Source/Dispatcher/Public/InterchangeDispatcher.h
+++ b/Engine/Plugins/Runtime/Interchange/Source/Dispatcher/Public/InterchangeDispatcher.h
@@ -32,6 +32,15 @@ namespace UE
 			void WaitAllTaskToCompleteExecution();
 			bool IsOver();
 
+			void SetInterchangeWorkerFatalError(FString& ErrorMessage)
+			{
+				InterchangeWorkerFatalError = MoveTemp(ErrorMessage);
+			}
+
+			FString GetInterchangeWorkerFatalError()
+			{
+				return InterchangeWorkerFatalError;
+			}
 		private:
 			void SpawnHandler();
 			bool IsHandlerAlive();
@@ -47,6 +56,8 @@ namespace UE
 			/** Path where the result files are dump */
 			FString ResultFolder;
 
+			FString InterchangeWorkerFatalError;
+
 			// Workers
 			TUniquePtr<FInterchangeWorkerHandler> WorkerHandler = nullptr;
 		};
diff --git a/Engine/Plugins/Runtime/Interchange/Source/Dispatcher/Public/InterchangeWorkerHandler.h b/Engine/Plugins/Runtime/Interchange/Source/Dispatcher/Public/InterchangeWorkerHandler.h
--- a/Engine/Plugins/Runtime/Interchange/Source/Dispatcher/Public/InterchangeWorkerHandler.h
+++ b/Engine/Plugins/Runtime/Interchange/Source/Dispatcher/Public/InterchangeWorkerHandler.h
@@ -61,6 +61,7 @@ namespace UE
 			void ValidateConnection();
 
 			void ProcessCommand(FPingCommand& PingCommand);
+			void ProcessCommand(FErrorCommand& ErrorCommand);
 			void ProcessCommand(FCompletedTaskCommand& RunTaskCommand);
 			void ProcessCommand(FCompletedQueryTaskProgressCommand& CompletedQueryTaskProgressCommand);
 			const TCHAR* EWorkerErrorStateAsString(EWorkerErrorState e);
diff --git a/Engine/Plugins/Runtime/Interchange/Source/Import/Private/Fbx/InterchangeFbxTranslator.cpp b/Engine/Plugins/Runtime/Interchange/Source/Import/Private/Fbx/InterchangeFbxTranslator.cpp
--- a/Engine/Plugins/Runtime/Interchange/Source/Import/Private/Fbx/InterchangeFbxTranslator.cpp
+++ b/Engine/Plugins/Runtime/Interchange/Source/Import/Private/Fbx/InterchangeFbxTranslator.cpp
@@ -88,7 +88,13 @@ bool UInterchangeFbxTranslator::Translate(UInterchangeBaseNodeContainer& BaseNod
 
 	//Blocking call until all tasks are executed
 	Dispatcher->WaitAllTaskToCompleteExecution();
-		
+	
+	FString WorkerFatalError = Dispatcher->GetInterchangeWorkerFatalError();
+	if (!WorkerFatalError.IsEmpty())
+	{
+		AddMessage(UInterchangeResult::FromJson(WorkerFatalError));
+	}
+
 	UE::Interchange::ETaskState TaskState;
 	FString JsonResult;
 	TArray<FString> JsonMessages;
diff --git a/Engine/Source/Programs/InterchangeWorker/Private/InterchangeWorker.cpp b/Engine/Source/Programs/InterchangeWorker/Private/InterchangeWorker.cpp
--- a/Engine/Source/Programs/InterchangeWorker/Private/InterchangeWorker.cpp
+++ b/Engine/Source/Programs/InterchangeWorker/Private/InterchangeWorker.cpp
@@ -3,8 +3,10 @@
 #include "InterchangeWorker.h"
 
 #include "HAL/PlatformProcess.h"
+#include "InterchangeResultsContainer.h"
 #include "InterchangeWorkerImpl.h"
 #include "RequiredProgramMainCPPInclude.h"
+#include "UObject/UObjectGlobals.h"
 
 IMPLEMENT_APPLICATION(InterchangeWorker, "InterchangeWorker");
 DEFINE_LOG_CATEGORY(LogInterchangeWorker);
@@ -53,19 +55,47 @@ int32 Main(int32 Argc, TCHAR * Argv[])
 	int32 Major = 0;
 	int32 Minor = 0;
 	int32 Patch = 0;
-	if (!UE::Interchange::DispatcherCommandVersion::FromString(InterchangeDispatcherVersion, Major, Minor, Patch))
+	bool bLWCDisabled = true;
+
+	FString WorkerVersionError;
+	if (!UE::Interchange::DispatcherCommandVersion::FromString(InterchangeDispatcherVersion, Major, Minor, Patch, bLWCDisabled))
 	{
-		UE_LOG(LogInterchangeWorker, Error, TEXT("Incompatible interchange dispatcher version string command argument."));
-		return EXIT_FAILURE;
+		WorkerVersionError = TEXT("Incompatible interchange dispatcher version string command argument.");
+		WorkerVersionError += TEXT(" Editor Version: ");
+		WorkerVersionError += InterchangeDispatcherVersion;
+		WorkerVersionError += TEXT(" Worker Version: ");
+		WorkerVersionError += UE::Interchange::DispatcherCommandVersion::ToString();
+		UE_LOG(LogInterchangeWorker, Error, TEXT("%s"), *WorkerVersionError);
 	}
- 	if (!UE::Interchange::DispatcherCommandVersion::IsAPICompatible(Major, Minor, Patch))
+
+ 	if (!UE::Interchange::DispatcherCommandVersion::IsAPICompatible(Major, Minor, Patch, bLWCDisabled))
  	{
- 		UE_LOG(LogInterchangeWorker, Error, TEXT("Incompatible interchange dispatcher version. Please recompile InterchangeWorker target."));
- 		return EXIT_FAILURE;
+		WorkerVersionError = TEXT("Incompatible interchange dispatcher API version. Please recompile InterchangeWorker target.");
+		WorkerVersionError += TEXT("\n\tEditor Version: ");
+		WorkerVersionError += InterchangeDispatcherVersion;
+		WorkerVersionError += TEXT("\n\tWorker Version: ");
+		WorkerVersionError += UE::Interchange::DispatcherCommandVersion::ToString();
+ 		UE_LOG(LogInterchangeWorker, Error, TEXT("%s"), *WorkerVersionError);
  	}
 
+	if (!WorkerVersionError.IsEmpty())
+	{
+		UInterchangeResultError_Generic* Message = NewObject<UInterchangeResultError_Generic>(GetTransientPackage());
+		Message->SourceAssetName = TEXT("InterchangeWorkerProgram");
+		Message->Text = FText::FromString(WorkerVersionError);
+		WorkerVersionError = Message->ToJson();
+		Message->ClearFlags(RF_Standalone);
+		Message->ClearInternalFlags(EInternalObjectFlags::Async | EInternalObjectFlags::AsyncLoading);
+		//Should be garbage collectable now
+	}
+
 	FInterchangeWorkerImpl Worker(FCString::Atoi(*ServerPID), FCString::Atoi(*ServerPort), ResultFolder);
-	Worker.Run();
+	Worker.Run(WorkerVersionError);
+
+	if (!WorkerVersionError.IsEmpty())
+	{
+		return EXIT_FAILURE;
+	}
 
 	return EXIT_SUCCESS;
 }
diff --git a/Engine/Source/Programs/InterchangeWorker/Private/InterchangeWorkerImpl.cpp b/Engine/Source/Programs/InterchangeWorker/Private/InterchangeWorkerImpl.cpp
--- a/Engine/Source/Programs/InterchangeWorker/Private/InterchangeWorkerImpl.cpp
+++ b/Engine/Source/Programs/InterchangeWorker/Private/InterchangeWorkerImpl.cpp
@@ -30,7 +30,7 @@ FInterchangeWorkerImpl::FInterchangeWorkerImpl(int32 InServerPID, int32 InServer
 	}
 }
 
-bool FInterchangeWorkerImpl::Run()
+bool FInterchangeWorkerImpl::Run(const FString& WorkerVersionError)
 {
 	UE_LOG(LogInterchangeWorker, Verbose, TEXT("connect to %d..."), ServerPort);
 	bool bConnected = NetworkInterface.Connect(TEXT("Interchange Worker"), ServerPort, Config::ConnectTimeout_s);
@@ -45,6 +45,29 @@ bool FInterchangeWorkerImpl::Run()
 		return false;
 	}
 
+	const bool bVersionError = !WorkerVersionError.IsEmpty();
+	if (bVersionError)
+	{
+		FErrorCommand ErrorCmd;
+		ErrorCmd.ErrorMessage = WorkerVersionError;
+		CommandIO.SendCommand(ErrorCmd, Config::SendCommandTimeout_s);
+		//We want to time out after maximum of 5 seconds
+		double TimeOut = 5.0;
+		while (TimeOut > 0)
+		{
+			if (TSharedPtr<ICommand> Command = CommandIO.GetNextCommand(0.02))
+			{
+				if(Command->GetType() == ECommandId::Terminate)
+				{
+					UE_LOG(LogInterchangeWorker, Verbose, TEXT("Terminate command received. Exiting."));
+					break;
+				}
+			}
+			TimeOut -= 0.02;
+		}
+		return false;
+	}
+
 	InitiatePing();
 
 	bool bIsRunning = true;
diff --git a/Engine/Source/Programs/InterchangeWorker/Private/InterchangeWorkerImpl.h b/Engine/Source/Programs/InterchangeWorker/Private/InterchangeWorkerImpl.h
--- a/Engine/Source/Programs/InterchangeWorker/Private/InterchangeWorkerImpl.h
+++ b/Engine/Source/Programs/InterchangeWorker/Private/InterchangeWorkerImpl.h
@@ -16,7 +16,7 @@ class FInterchangeWorkerImpl
 {
 public:
 	FInterchangeWorkerImpl(int32 InServerPID, int32 InServerPort, FString& InResultFolder);
-	bool Run();
+	bool Run(const FString& WorkerVersionError);
 
 private:
 	void InitiatePing();
