fix monolithic build

diff --git a/Engine/Plugins/Experimental/Interchange/Runtime/Interchange.uplugin b/Engine/Plugins/Experimental/Interchange/Runtime/Interchange.uplugin
--- a/Engine/Plugins/Experimental/Interchange/Runtime/Interchange.uplugin
+++ b/Engine/Plugins/Experimental/Interchange/Runtime/Interchange.uplugin
@@ -53,12 +53,16 @@
 		},
 		{
 			"Name": "InterchangeFbxParser",
-			"Type": "Editor",
+			"Type": "Runtime",
 			"LoadingPhase": "Default",
 			"PlatformAllowList": [
 				"Win64",
 				"Linux",
 				"Mac"
+			],
+			"TargetAllowList": [
+				"Editor",
+				"Program"
 			]
 		}
 	],
diff --git a/Engine/Plugins/Experimental/Interchange/Runtime/Source/Parsers/Fbx/InterchangeFbxParser.Build.cs b/Engine/Plugins/Experimental/Interchange/Runtime/Source/Parsers/Fbx/InterchangeFbxParser.Build.cs
--- a/Engine/Plugins/Experimental/Interchange/Runtime/Source/Parsers/Fbx/InterchangeFbxParser.Build.cs
+++ b/Engine/Plugins/Experimental/Interchange/Runtime/Source/Parsers/Fbx/InterchangeFbxParser.Build.cs
@@ -8,7 +8,17 @@ namespace UnrealBuildTool.Rules
 	{
 		public InterchangeFbxParser(ReadOnlyTargetRules Target) : base(Target)
 		{
-			bEnableExceptions = true;
+			if (Target.Type != TargetType.Editor && Target.Type != TargetType.Program)
+			{
+				throw new BuildException("InterchangeFbxParser module can be build only on editor or program target.");
+			}
+
+			if (Target.Platform != UnrealTargetPlatform.Win64 &&
+				Target.Platform != UnrealTargetPlatform.Linux &&
+				Target.Platform != UnrealTargetPlatform.Mac)
+			{
+				throw new BuildException("InterchangeFbxParser module do not support target platform {0}.", Target.Platform.ToString());
+			}
 
 			PublicDependencyModuleNames.AddRange(
 				new string[]
diff --git a/Engine/Plugins/MeshPainting/MeshPainting.uplugin b/Engine/Plugins/MeshPainting/MeshPainting.uplugin
--- a/Engine/Plugins/MeshPainting/MeshPainting.uplugin
+++ b/Engine/Plugins/MeshPainting/MeshPainting.uplugin
@@ -32,7 +32,7 @@
 			"Enabled": true
 		},
 		{
-			"Name": "Interchange",
+			"Name": "InterchangeEditor",
 			"Enabled": true
 		}
 	],
diff --git a/Engine/Source/Programs/InterchangeWorker/InterchangeWorker.Build.cs b/Engine/Source/Programs/InterchangeWorker/InterchangeWorker.Build.cs
--- a/Engine/Source/Programs/InterchangeWorker/InterchangeWorker.Build.cs
+++ b/Engine/Source/Programs/InterchangeWorker/InterchangeWorker.Build.cs
@@ -7,6 +7,12 @@ public class InterchangeWorker : ModuleRules
 {
 	public InterchangeWorker(ReadOnlyTargetRules Target) : base(Target)
 	{
+		if (Target.Platform != UnrealTargetPlatform.Win64 &&
+			Target.Platform != UnrealTargetPlatform.Linux &&
+			Target.Platform != UnrealTargetPlatform.Mac)
+		{
+			throw new BuildException("InterchangeWorker program do not support target platform {0}.", Target.Platform.ToString());
+		}
 
 		PublicIncludePaths.Add("Runtime/Launch/Public");
 		PrivateIncludePaths.Add("Runtime/Launch/Private");
@@ -17,32 +23,21 @@ public class InterchangeWorker : ModuleRules
 				"ApplicationCore",
 				"Core",
 				"CoreUObject",
+				"InterchangeCore",
 				"InterchangeDispatcher",
+				"InterchangeFbxParser",
+				"InterchangeNodes",
 				"Json",
 				"Projects",
 				"Sockets",
 			}
 		);
 
-		if (Target.Platform == UnrealTargetPlatform.Win64 ||
-			Target.Platform == UnrealTargetPlatform.Linux ||
-			Target.Platform == UnrealTargetPlatform.Mac)
-		{
-			PrivateDependencyModuleNames.AddRange(
-				new string[]
-				{
-					"InterchangeCore",
-					"InterchangeFbxParser",
-					"InterchangeNodes"
-				}
-			);
-
-			AddEngineThirdPartyPrivateStaticDependencies(Target,
-				new string[]
-				{
-					"FBX"
-				}
-			);
-		}
+		AddEngineThirdPartyPrivateStaticDependencies(Target,
+			new string[]
+			{
+				"FBX"
+			}
+		);
 	}
 }
