Compiler fixup from variable moving

diff --git a/Engine/Platforms/Hololens/Source/Runtime/Launch/Private/HoloLensLaunch.cpp b/Engine/Platforms/Hololens/Source/Runtime/Launch/Private/HoloLensLaunch.cpp
--- a/Engine/Platforms/Hololens/Source/Runtime/Launch/Private/HoloLensLaunch.cpp
+++ b/Engine/Platforms/Hololens/Source/Runtime/Launch/Private/HoloLensLaunch.cpp
@@ -1159,7 +1159,6 @@ void appHoloLensEarlyInit()
 
 /** The global EngineLoop instance */
 FEngineLoop	GEngineLoop;
-bool GIsConsoleExecutable = false;
 
 /** 
  * PreInits the engine loop 
diff --git a/Engine/Source/Programs/InterchangeWorker/Private/InterchangeWorker.cpp b/Engine/Source/Programs/InterchangeWorker/Private/InterchangeWorker.cpp
--- a/Engine/Source/Programs/InterchangeWorker/Private/InterchangeWorker.cpp
+++ b/Engine/Source/Programs/InterchangeWorker/Private/InterchangeWorker.cpp
@@ -13,8 +13,6 @@ DEFINE_LOG_CATEGORY(LogInterchangeWorker);
 
 #define EXIT_MISSING_CORETECH_MODULE 2
 
-bool GIsConsoleExecutable = true;
-
 void GetParameter(int32 Argc, TCHAR* Argv[], const FString& InParam, FString& OutValue)
 {
 	OutValue.Empty();
