Fixed crash on exit of InterchangeWorker

diff --git a/Engine/Source/Programs/InterchangeWorker/Private/InterchangeWorker.cpp b/Engine/Source/Programs/InterchangeWorker/Private/InterchangeWorker.cpp
--- a/Engine/Source/Programs/InterchangeWorker/Private/InterchangeWorker.cpp
+++ b/Engine/Source/Programs/InterchangeWorker/Private/InterchangeWorker.cpp
@@ -122,6 +122,7 @@ INT32_MAIN_INT32_ARGC_TCHAR_ARGV()
 	FEngineLoop::AppPreExit();
 	FModuleManager::Get().UnloadModulesAtShutdown();
 	FEngineLoop::AppExit();
+	FPlatformMisc::RequestExit(false);
 
 	return ReturnCode;
 }
