Added missing file from CL 45200263

#rnx

diff --git a/Engine/Source/Programs/InterchangeWorker/Private/InterchangeWorkerImpl.h b/Engine/Source/Programs/InterchangeWorker/Private/InterchangeWorkerImpl.h
--- a/Engine/Source/Programs/InterchangeWorker/Private/InterchangeWorkerImpl.h
+++ b/Engine/Source/Programs/InterchangeWorker/Private/InterchangeWorkerImpl.h
@@ -1,6 +1,7 @@
 // Copyright Epic Games, Inc. All Rights Reserved.
 #pragma once
 
+#include "FBX/InterchangeDispatcherFBXTasks.h"
 #include "HAL/CriticalSection.h"
 #include "HAL/PlatformTime.h"
 #include "HAL/Thread.h"
@@ -25,7 +26,7 @@ private:
 	void ProcessCommand(const TSharedPtr<UE::Interchange::ICommand> Command, const FString& ThreadName);
 	void ProcessCommand(const UE::Interchange::FQueryTaskProgressCommand& QueryTaskProgressCommand);
 
-	UE::Interchange::ETaskState LoadFbxFile(const UE::Interchange::FJsonLoadSourceCmd& LoadSourceCommand, FString& OutJSonResult, TArray<FString>& OutJSonMessages);
+	UE::Interchange::ETaskState LoadFbxFile(const UE::Interchange::FJsonFBXLoadSourceCmd& LoadSourceCommand, FString& OutJSonResult, TArray<FString>& OutJSonMessages);
 	UE::Interchange::ETaskState FetchFbxPayload(const UE::Interchange::FJsonFetchPayloadCmd& FetchPayloadCommand, FString& OutJSonResult, TArray<FString>& OutJSonMessages);
 	UE::Interchange::ETaskState FetchFbxPayload(const UE::Interchange::FJsonFetchMeshPayloadCmd& FetchMeshPayloadCommand, FString& OutJSonResult, TArray<FString>& OutJSonMessages);
 	UE::Interchange::ETaskState FetchFbxPayload(const UE::Interchange::FJsonFetchAnimationQueriesCmd& FetchAnimationBakeTransformPayloadCommand, FString& OutJSonResult, TArray<FString>& OutJSonMessages);
